Added new Blendable called "LightPropagationVolume Blendable" which can be hooked up to the PostProcessVolume/Misc/Blendables to override LPV settings.
This work is an example on how to make rendering settings outside of the engine still getting the convenient blending functionality of the PostProcess settings.

todo:
* Blueprint scripting (dynamic asset generation?)
* Improve the UI to create settings without using the content browser (asset living in the level).

later todo:
* Remove the LPV settings in the struct (content changes might be needed, if we don't provide an automatic upgrade for those settings).
* We can expose the BlendableManager to other areas (gameplay/tick/blueprint).
* Break out more/all from the post process struct
* We might change the name of the PostProcessVolume as it no longer only affects post processing.

diff --git a/Engine/Plugins/Blendables/LightPropagationVolume/LightPropagationVolume.uplugin b/Engine/Plugins/Blendables/LightPropagationVolume/LightPropagationVolume.uplugin
new file mode 100644
--- /dev/null
+++ b/Engine/Plugins/Blendables/LightPropagationVolume/LightPropagationVolume.uplugin
@@ -0,0 +1,27 @@
+{
+    "FileVersion": 3,
+
+    "FriendlyName": "LightPropagationVolume",
+    "Version": 1,
+    "VersionName": "1.0",
+    "CreatedBy": "Epic Games, Inc.",
+    "CreatedByURL": "http://epicgames.com",
+    "EngineVersion": "4.2.0",
+    "Description": "todo",
+    "Category": "3D",
+    "EnabledByDefault": true,
+
+    "Modules": [
+        {
+            "Name": "LightPropagationVolumeRuntime",
+            "Type": "Runtime",
+            "LoadingPhase": "PreDefault"
+        },
+        {
+			"Name" : "LightPropagationVolumeEditor",
+			"Type" : "Editor"
+		},
+    ],
+
+	"CanContainContent" : true
+}
\ No newline at end of file
diff --git a/Engine/Plugins/Blendables/LightPropagationVolume/Source/LightPropagationVolumeEditor/Classes/LightPropagationVolumeBlendableFactory.h b/Engine/Plugins/Blendables/LightPropagationVolume/Source/LightPropagationVolumeEditor/Classes/LightPropagationVolumeBlendableFactory.h
new file mode 100644
--- /dev/null
+++ b/Engine/Plugins/Blendables/LightPropagationVolume/Source/LightPropagationVolumeEditor/Classes/LightPropagationVolumeBlendableFactory.h
@@ -0,0 +1,15 @@
+// Copyright 1998-2015 Epic Games, Inc. All Rights Reserved.
+
+#pragma once
+
+#include "LightPropagationVolumeBlendableFactory.generated.h"
+
+UCLASS(hidecategories=Object)
+class ULightPropagationVolumeBlendableFactory : public UFactory
+{
+	GENERATED_UCLASS_BODY()
+
+	// UFactory interface
+	virtual UObject* FactoryCreateNew(UClass* Class, UObject* InParent, FName Name, EObjectFlags Flags, UObject* Context, FFeedbackContext* Warn) override;
+	// End of UFactory interface
+};
diff --git a/Engine/Plugins/Blendables/LightPropagationVolume/Source/LightPropagationVolumeEditor/LightPropagationVolumeEditor.Build.cs b/Engine/Plugins/Blendables/LightPropagationVolume/Source/LightPropagationVolumeEditor/LightPropagationVolumeEditor.Build.cs
new file mode 100644
--- /dev/null
+++ b/Engine/Plugins/Blendables/LightPropagationVolume/Source/LightPropagationVolumeEditor/LightPropagationVolumeEditor.Build.cs
@@ -0,0 +1,25 @@
+// Copyright 1998-2015 Epic Games, Inc. All Rights Reserved.
+
+using UnrealBuildTool;
+
+public class LightPropagationVolumeEditor : ModuleRules
+{
+    public LightPropagationVolumeEditor(TargetInfo Target)
+	{
+		PrivateDependencyModuleNames.AddRange(
+			new string[] {	
+				"Core",
+				"CoreUObject",
+                "UnrealEd",         // for FAssetEditorManager
+				"AssetTools",       // class FAssetTypeActions_Base
+				"LightPropagationVolumeRuntime",
+			}
+		);
+
+        PrivateIncludePathModuleNames.AddRange(
+            new string[] {
+				"UnrealEd",
+			}
+        );
+	}
+}
diff --git a/Engine/Plugins/Blendables/LightPropagationVolume/Source/LightPropagationVolumeEditor/Private/AssetTypeActions_LightPropagationVolumeBlendable.h b/Engine/Plugins/Blendables/LightPropagationVolume/Source/LightPropagationVolumeEditor/Private/AssetTypeActions_LightPropagationVolumeBlendable.h
new file mode 100644
--- /dev/null
+++ b/Engine/Plugins/Blendables/LightPropagationVolume/Source/LightPropagationVolumeEditor/Private/AssetTypeActions_LightPropagationVolumeBlendable.h
@@ -0,0 +1,27 @@
+// Copyright 1998-2012 Epic Games, Inc. All Rights Reserved.
+
+#pragma once
+
+#include "AssetTypeActions_Base.h"
+#include "LightPropagationVolumeBlendable.h"
+
+class FAssetTypeActions_LightPropagationVolumeBlendable : public FAssetTypeActions_Base
+{
+public:
+	FAssetTypeActions_LightPropagationVolumeBlendable(EAssetTypeCategories::Type InAssetCategoryBit)
+		: AssetCategoryBit(InAssetCategoryBit)
+	{
+	}
+
+	// IAssetTypeActions Implementation
+	virtual FText GetName() const override { return NSLOCTEXT("AssetTypeActions", "AssetTypeActions_LightPropagationVolumeBlendable", "LightPropagationVolume Blendable"); }
+	virtual FColor GetTypeColor() const override { return FColor(128, 0, 200); }
+	virtual bool HasActions ( const TArray<UObject*>& InObjects ) const override { return false; }
+	virtual UClass* GetSupportedClass() const override { return ULightPropagationVolumeBlendable::StaticClass(); }
+	virtual bool CanFilter() override { return true; }
+	virtual uint32 GetCategories() override { return AssetCategoryBit; }
+
+private:
+
+	EAssetTypeCategories::Type AssetCategoryBit;
+};
\ No newline at end of file
diff --git a/Engine/Plugins/Blendables/LightPropagationVolume/Source/LightPropagationVolumeEditor/Private/LightPropagationVolumeEditorModule.cpp b/Engine/Plugins/Blendables/LightPropagationVolume/Source/LightPropagationVolumeEditor/Private/LightPropagationVolumeEditorModule.cpp
new file mode 100644
--- /dev/null
+++ b/Engine/Plugins/Blendables/LightPropagationVolume/Source/LightPropagationVolumeEditor/Private/LightPropagationVolumeEditorModule.cpp
@@ -0,0 +1,77 @@
+// Copyright 1998-2015 Epic Games, Inc. All Rights Reserved.
+
+/*=============================================================================
+	LightPropagationVolumeModule.cpp: Module encapsulates the asset type for blending LightPropagationVolume settings
+=============================================================================*/
+
+#include "Core.h"
+#include "AssetTypeActions_LightPropagationVolumeBlendable.h"
+#include "LightPropagationVolumeBlendableFactory.h"
+#include "LightPropagationVolumeBlendable.h"
+
+#define LOCTEXT_NAMESPACE "LightPropagationVolume"
+
+/////////////////////////////////////////////////////
+// ULightPropagationVolumeBlendableFactory
+
+ULightPropagationVolumeBlendableFactory::ULightPropagationVolumeBlendableFactory(const FObjectInitializer& ObjectInitializer)
+	: Super(ObjectInitializer)
+{
+	bCreateNew = true;
+	bEditAfterNew = true;
+	SupportedClass = ULightPropagationVolumeBlendable::StaticClass();
+}
+
+UObject* ULightPropagationVolumeBlendableFactory::FactoryCreateNew(UClass* Class, UObject* InParent, FName Name, EObjectFlags Flags, UObject* Context, FFeedbackContext* Warn)
+{
+	return NewObject<ULightPropagationVolumeBlendable>(InParent, Class, Name, Flags);
+}
+
+//////////////////////////////////////////////////////////////////////////
+// FLightPropagationVolumeModule
+
+class FLightPropagationVolumeModule : public IModuleInterface
+{
+	virtual void StartupModule() override
+	{
+		// Register asset types
+		IAssetTools& AssetTools = FModuleManager::LoadModuleChecked<FAssetToolsModule>("AssetTools").Get();
+
+ 		EAssetTypeCategories::Type AssetCategoryBit = AssetTools.RegisterAdvancedAssetCategory(FName(TEXT("Blendables")), LOCTEXT("BlendablesAssetCategory", "Blendables"));
+		RegisterAssetTypeAction(AssetTools, MakeShareable(new FAssetTypeActions_LightPropagationVolumeBlendable(AssetCategoryBit)));
+	}
+
+	virtual void ShutdownModule() override
+	{
+		// Unregister all the asset types that we registered
+		if (FModuleManager::Get().IsModuleLoaded("AssetTools"))
+		{
+			IAssetTools& AssetTools = FModuleManager::GetModuleChecked<FAssetToolsModule>("AssetTools").Get();
+
+			for (int32 Index = 0; Index < CreatedAssetTypeActions.Num(); ++Index)
+			{
+				AssetTools.UnregisterAssetTypeActions(CreatedAssetTypeActions[Index].ToSharedRef());
+			}
+		}
+		CreatedAssetTypeActions.Empty();
+	}
+
+private:
+	
+	/** All created asset type actions.  Cached here so that we can unregister them during shutdown. */
+	TArray< TSharedPtr<IAssetTypeActions> > CreatedAssetTypeActions;
+
+	void RegisterAssetTypeAction(IAssetTools& AssetTools, TSharedRef<IAssetTypeActions> Action)
+	{
+		AssetTools.RegisterAssetTypeActions(Action);
+		CreatedAssetTypeActions.Add(Action);
+	}
+};
+
+//////////////////////////////////////////////////////////////////////////
+
+IMPLEMENT_MODULE(FLightPropagationVolumeModule, FLightPropagationVolumeModule);
+
+//////////////////////////////////////////////////////////////////////////
+
+#undef LOCTEXT_NAMESPACE
\ No newline at end of file
diff --git a/Engine/Plugins/Blendables/LightPropagationVolume/Source/LightPropagationVolumeRuntime/LightPropagationVolumeRuntime.Build.cs b/Engine/Plugins/Blendables/LightPropagationVolume/Source/LightPropagationVolumeRuntime/LightPropagationVolumeRuntime.Build.cs
new file mode 100644
--- /dev/null
+++ b/Engine/Plugins/Blendables/LightPropagationVolume/Source/LightPropagationVolumeRuntime/LightPropagationVolumeRuntime.Build.cs
@@ -0,0 +1,20 @@
+// Copyright 1998-2015 Epic Games, Inc. All Rights Reserved.
+
+using UnrealBuildTool;
+
+public class LightPropagationVolumeRuntime : ModuleRules
+{
+	public LightPropagationVolumeRuntime(TargetInfo Target)
+	{
+		PrivateDependencyModuleNames.AddRange(
+			new string[] {	
+				"Core",
+				"CoreUObject",
+				"Engine",           // FBlendableManager
+//				"RHI",
+//				"RenderCore",
+//				"ShaderCore",
+			}
+		);
+	}
+}
diff --git a/Engine/Plugins/Blendables/LightPropagationVolume/Source/LightPropagationVolumeRuntime/Private/LightPropagationVolumeRuntimeModule.cpp b/Engine/Plugins/Blendables/LightPropagationVolume/Source/LightPropagationVolumeRuntime/Private/LightPropagationVolumeRuntimeModule.cpp
new file mode 100644
--- /dev/null
+++ b/Engine/Plugins/Blendables/LightPropagationVolume/Source/LightPropagationVolumeRuntime/Private/LightPropagationVolumeRuntimeModule.cpp
@@ -0,0 +1,51 @@
+// Copyright 1998-2015 Epic Games, Inc. All Rights Reserved.
+
+/*=============================================================================
+	LightPropagationVolumeRuntimeModule.cpp: Module encapsulates the LightPropagationVolume settings
+=============================================================================*/
+
+#include "Core.h"
+#include "LightPropagationVolumeBlendable.h"
+
+IMPLEMENT_MODULE(FDefaultModuleImpl, LightPropagationVolumeRuntime);
+
+//////////////////////////////////////////////////////////////////////////
+// ULightPropagationVolumeBlendable
+
+ULightPropagationVolumeBlendable::ULightPropagationVolumeBlendable(const FObjectInitializer& ObjectInitializer)
+	: Super(ObjectInitializer)
+{
+	BlendWeight = 1.0f;
+}
+
+#define LERP_PP(NAME) if(Settings.bOverride_ ## NAME)	Dest.NAME = FMath::Lerp(Dest.NAME, Settings.NAME, Weight);
+
+void ULightPropagationVolumeBlendable::OverrideBlendableSettings(FSceneView& View, float InWeight) const
+{
+	float Weight = FMath::Clamp(BlendWeight, 0.0f, 1.0f) * InWeight;
+
+	FLightPropagationVolumeSettings& Dest = View.FinalPostProcessSettings.BlendableManager.GetSingleFinalData<FLightPropagationVolumeSettings>();
+	const FSceneViewFamily* Family = View.Family;
+		
+	LERP_PP(LPVIntensity);
+	LERP_PP(LPVSecondaryOcclusionIntensity);
+	LERP_PP(LPVSecondaryBounceIntensity);
+	LERP_PP(LPVVplInjectionBias);
+	LERP_PP(LPVGeometryVolumeBias);
+	LERP_PP(LPVEmissiveInjectionIntensity);
+	LERP_PP(LPVDirectionalOcclusionIntensity);
+	LERP_PP(LPVDirectionalOcclusionRadius);
+	LERP_PP(LPVDiffuseOcclusionExponent);
+	LERP_PP(LPVSpecularOcclusionExponent);
+	LERP_PP(LPVDiffuseOcclusionIntensity);
+	LERP_PP(LPVSpecularOcclusionIntensity);
+
+	if(Settings.bOverride_LPVSize && Weight > 0)
+	{
+		// we don't want to lerp this property
+		Dest.LPVSize = Settings.LPVSize;
+	}
+}
+
+
+
diff --git a/Engine/Plugins/Blendables/LightPropagationVolume/Source/LightPropagationVolumeRuntime/Public/LightPropagationVolumeBlendable.h b/Engine/Plugins/Blendables/LightPropagationVolume/Source/LightPropagationVolumeRuntime/Public/LightPropagationVolumeBlendable.h
new file mode 100644
--- /dev/null
+++ b/Engine/Plugins/Blendables/LightPropagationVolume/Source/LightPropagationVolumeRuntime/Public/LightPropagationVolumeBlendable.h
@@ -0,0 +1,165 @@
+// Copyright 1998-2012 Epic Games, Inc. All Rights Reserved.
+
+#pragma once
+
+#include "Engine/BlendableInterface.h"
+#include "Engine/DataAsset.h"
+#include "LightPropagationVolumeBlendable.generated.h"
+
+USTRUCT(BlueprintType)
+struct FLightPropagationVolumeSettings
+{
+	GENERATED_USTRUCT_BODY()
+
+	UPROPERTY(BlueprintReadWrite, Category=Overrides, meta=(PinHiddenByDefault))
+	uint32 bOverride_LPVIntensity:1;
+
+	UPROPERTY(BlueprintReadWrite, Category=Overrides, meta=(PinHiddenByDefault))
+	uint32 bOverride_LPVDirectionalOcclusionIntensity:1;
+
+	UPROPERTY(BlueprintReadWrite, Category=Overrides, meta=(PinHiddenByDefault))
+	uint32 bOverride_LPVDirectionalOcclusionRadius:1;
+
+	UPROPERTY(BlueprintReadWrite, Category=Overrides, meta=(PinHiddenByDefault))
+	uint32 bOverride_LPVDiffuseOcclusionExponent:1;
+
+	UPROPERTY(BlueprintReadWrite, Category=Overrides, meta=(PinHiddenByDefault))
+	uint32 bOverride_LPVSpecularOcclusionExponent:1;
+
+	UPROPERTY(BlueprintReadWrite, Category=Overrides, meta=(PinHiddenByDefault))
+	uint32 bOverride_LPVDiffuseOcclusionIntensity:1;
+
+	UPROPERTY(BlueprintReadWrite, Category=Overrides, meta=(PinHiddenByDefault))
+	uint32 bOverride_LPVSpecularOcclusionIntensity:1;
+
+	UPROPERTY(BlueprintReadWrite, Category=Overrides, meta=(PinHiddenByDefault))
+	uint32 bOverride_LPVSize:1;
+
+	UPROPERTY(BlueprintReadWrite, Category=Overrides, meta=(PinHiddenByDefault))
+	uint32 bOverride_LPVSecondaryOcclusionIntensity:1;
+
+	UPROPERTY(BlueprintReadWrite, Category=Overrides, meta=(PinHiddenByDefault))
+	uint32 bOverride_LPVSecondaryBounceIntensity:1;
+
+	UPROPERTY(BlueprintReadWrite, Category=Overrides, meta=(PinHiddenByDefault))
+	uint32 bOverride_LPVGeometryVolumeBias:1;
+
+	UPROPERTY(BlueprintReadWrite, Category=Overrides, meta=(PinHiddenByDefault))
+	uint32 bOverride_LPVVplInjectionBias:1;
+
+	UPROPERTY(BlueprintReadWrite, Category=Overrides, meta=(PinHiddenByDefault))
+	uint32 bOverride_LPVEmissiveInjectionIntensity:1;
+
+	/** How strong the dynamic GI from the LPV should be. 0.0 is off, 1.0 is the "normal" value, but higher values can be used to boost the effect*/
+	UPROPERTY(interp, BlueprintReadWrite, Category=LightPropagationVolume, meta=(editcondition = "bOverride_LPVIntensity", UIMin = "0", UIMax = "20", DisplayName = "Intensity") )
+	float LPVIntensity;
+
+	/** Bias applied to light injected into the LPV in cell units. Increase to reduce bleeding through thin walls*/
+	UPROPERTY(interp, BlueprintReadWrite, Category=LightPropagationVolume, AdvancedDisplay, meta=(editcondition = "bOverride_LPVVplInjectionBias", UIMin = "0", UIMax = "2", DisplayName = "Light Injection Bias") )
+	float LPVVplInjectionBias;
+
+	/** The size of the LPV volume, in Unreal units*/
+	UPROPERTY(EditAnywhere, BlueprintReadWrite, Category=LightPropagationVolume, meta=(editcondition = "bOverride_LPVSize", UIMin = "100", UIMax = "20000", DisplayName = "Size") )
+	float LPVSize;
+
+	/** Secondary occlusion strength (bounce light shadows). Set to 0 to disable*/
+	UPROPERTY(interp, BlueprintReadWrite, Category=LightPropagationVolume, meta=(editcondition = "bOverride_LPVSecondaryOcclusionIntensity", UIMin = "0", UIMax = "1", DisplayName = "Secondary Occlusion Intensity") )
+	float LPVSecondaryOcclusionIntensity;
+
+	/** Secondary bounce light strength (bounce light shadows). Set to 0 to disable*/
+	UPROPERTY(interp, BlueprintReadWrite, Category=LightPropagationVolume, AdvancedDisplay, meta=(editcondition = "bOverride_LPVSecondaryBounceIntensity", UIMin = "0", UIMax = "1", DisplayName = "Secondary Bounce Intensity") )
+	float LPVSecondaryBounceIntensity;
+
+	/** Bias applied to the geometry volume in cell units. Increase to reduce darkening due to secondary occlusion */
+	UPROPERTY(interp, BlueprintReadWrite, Category=LightPropagationVolume, AdvancedDisplay, meta=(editcondition = "bOverride_LPVGeometryVolumeBias", UIMin = "0", UIMax = "2", DisplayName = "Geometry Volume Bias"))
+	float LPVGeometryVolumeBias;
+
+	UPROPERTY(interp, BlueprintReadWrite, Category=LightPropagationVolume, AdvancedDisplay, meta=(editcondition = "bOverride_LPVEmissiveInjectionIntensity", UIMin = "0", UIMax = "20", DisplayName = "Emissive Injection Intensity") )
+	float LPVEmissiveInjectionIntensity;
+
+	/** Controls the amount of directional occlusion. Requires LPV. Values very close to 1.0 are recommended */
+	UPROPERTY(interp, BlueprintReadWrite, Category=LightPropagationVolume, meta=(editcondition = "bOverride_LPVDirectionalOcclusionIntensity", UIMin = "0", UIMax = "1", DisplayName = "Occlusion Intensity") )
+	float LPVDirectionalOcclusionIntensity;
+
+	/** Occlusion Radius - 16 is recommended for most scenes */
+	UPROPERTY(interp, BlueprintReadWrite, Category=LightPropagationVolume, AdvancedDisplay, meta=(editcondition = "bOverride_LPVDirectionalOcclusionRadius", UIMin = "1", UIMax = "16", DisplayName = "Occlusion Radius") )
+	float LPVDirectionalOcclusionRadius;
+
+	/** Diffuse occlusion exponent - increase for more contrast. 1 to 2 is recommended */
+	UPROPERTY(interp, BlueprintReadWrite, Category=LightPropagationVolume, meta=(editcondition = "bOverride_LPVDiffuseOcclusionExponent", UIMin = "0.5", UIMax = "5", DisplayName = "Diffuse occlusion exponent") )
+	float LPVDiffuseOcclusionExponent;
+
+	/** Specular occlusion exponent - increase for more contrast. 6 to 9 is recommended */
+	UPROPERTY(interp, BlueprintReadWrite, Category=LightPropagationVolume, meta=(editcondition = "bOverride_LPVSpecularOcclusionExponent", UIMin = "1", UIMax = "16", DisplayName = "Specular occlusion exponent") )
+	float LPVSpecularOcclusionExponent;
+
+	/** Diffuse occlusion intensity - higher values provide increased diffuse occlusion.*/
+	UPROPERTY(interp, BlueprintReadWrite, Category=LightPropagationVolume, AdvancedDisplay, meta=(editcondition = "bOverride_LPVDiffuseOcclusionIntensity", UIMin = "0", UIMax = "4", DisplayName = "Diffuse occlusion intensity") )
+	float LPVDiffuseOcclusionIntensity;
+
+	/** Specular occlusion intensity - higher values provide increased specular occlusion.*/
+	UPROPERTY(interp, BlueprintReadWrite, Category=LightPropagationVolume, AdvancedDisplay, meta=(editcondition = "bOverride_LPVSpecularOcclusionIntensity", UIMin = "0", UIMax = "4", DisplayName = "Specular occlusion intensity") )
+	float LPVSpecularOcclusionIntensity;
+
+	// good start values for a new volume, by default no value is overriding
+
+	FLightPropagationVolumeSettings()
+	{
+		// to set all bOverride_.. by default to false
+		FMemory::Memzero(this, sizeof(*this));
+
+		// default values:
+
+		LPVIntensity = 1.0f;
+		LPVSize = 5312.0f;
+		LPVSecondaryOcclusionIntensity = 0.0f;
+		LPVSecondaryBounceIntensity = 0.0f;
+		LPVVplInjectionBias = 0.64f;
+		LPVGeometryVolumeBias = 0.384f;
+		LPVEmissiveInjectionIntensity = 1.0f;
+		LPVDirectionalOcclusionIntensity = 0.0f;
+		LPVDirectionalOcclusionRadius = 8.0f;
+		LPVDiffuseOcclusionExponent = 1.0f;
+		LPVSpecularOcclusionExponent = 7.0f;
+		LPVDiffuseOcclusionIntensity = 1.0f;
+		LPVSpecularOcclusionIntensity = 1.0f;
+	}
+
+	/**
+	 * Used to define the values before any override happens.
+	 * Should be as neutral as possible.
+	 */		
+	void SetBaseValues()
+	{
+		*this = FLightPropagationVolumeSettings();
+	}
+	
+	// for type safety in FBlendableManager
+	static const FName& GetFName()
+	{
+		static const FName Name = FName(TEXT("FLightPropagationVolumeSettings"));
+
+		return Name;
+	}
+};
+
+UCLASS(MinimalAPI, Blueprintable)
+//class ULightPropagationVolumeBlendable : public UDataAsset, public IBlendableInterface
+class ULightPropagationVolumeBlendable : public UDataAsset, public IBlendableInterface
+{
+	GENERATED_UCLASS_BODY()
+
+	UPROPERTY(interp, Category=LightPropagationVolumeBlendable, meta=(ShowOnlyInnerProperties))
+	struct FLightPropagationVolumeSettings Settings;
+
+	/** 0:no effect, 1:full effect */
+	UPROPERTY(interp, Category=LightPropagationVolumeBlendable, BlueprintReadWrite, meta=(UIMin = "0.0", UIMax = "1.0"))
+	float BlendWeight;
+	
+	// Begin interface IBlendableInterface
+	virtual void OverrideBlendableSettings(class FSceneView& View, float InWeight) const override;
+	// End interface IBlendableInterface
+};
+
+
+
diff --git a/Engine/Source/Runtime/Engine/Engine.Build.cs b/Engine/Source/Runtime/Engine/Engine.Build.cs
--- a/Engine/Source/Runtime/Engine/Engine.Build.cs
+++ b/Engine/Source/Runtime/Engine/Engine.Build.cs
@@ -88,6 +88,8 @@ public class Engine : ModuleRules
 			}
 		);
 
+        PrivateIncludePathModuleNames.Add("LightPropagationVolumeRuntime");
+
 		if (Target.Type != TargetRules.TargetType.Server)
 		{
 			PrivateIncludePathModuleNames.AddRange(
diff --git a/Engine/Source/Runtime/Engine/Private/SceneView.cpp b/Engine/Source/Runtime/Engine/Private/SceneView.cpp
--- a/Engine/Source/Runtime/Engine/Private/SceneView.cpp
+++ b/Engine/Source/Runtime/Engine/Private/SceneView.cpp
@@ -17,6 +17,7 @@
 #include "Engine/TextureCube.h"
 #include "IHeadMountedDisplay.h"
 #include "Classes/Engine/RendererSettings.h"
+#include "LightPropagationVolumeBlendable.h"
 
 DEFINE_LOG_CATEGORY(LogBufferVisualization);
 
@@ -801,13 +802,6 @@ void FSceneView::OverridePostProcessSettings(const FPostProcessSettings& Src, fl
 	LERP_PP(BloomDirtMaskTint);
 	LERP_PP(AmbientCubemapIntensity);
 	LERP_PP(AmbientCubemapTint);
-	LERP_PP(LPVIntensity);
-	LERP_PP(LPVDirectionalOcclusionIntensity);
-	LERP_PP(LPVDirectionalOcclusionRadius);
-	LERP_PP(LPVSpecularOcclusionExponent);
-	LERP_PP(LPVDiffuseOcclusionExponent);
-	LERP_PP(LPVDiffuseOcclusionIntensity);
-	LERP_PP(LPVSpecularOcclusionIntensity);
 	LERP_PP(AutoExposureLowPercent);
 	LERP_PP(AutoExposureHighPercent);
 	LERP_PP(AutoExposureMinBrightness);
@@ -901,16 +895,6 @@ void FSceneView::OverridePostProcessSettings(const FPostProcessSettings& Src, fl
 		Dest.LensFlareBokehShape = Src.LensFlareBokehShape;
 	}
 
-	if(Src.bOverride_LPVSize)
-	{
-		Dest.LPVSize = Src.LPVSize;
-	}
-	LERP_PP( LPVSecondaryOcclusionIntensity );
-	LERP_PP( LPVSecondaryBounceIntensity );
-	LERP_PP( LPVVplInjectionBias );
-	LERP_PP( LPVGeometryVolumeBias );
-	LERP_PP( LPVEmissiveInjectionIntensity );
-
 	if(Src.bOverride_LensFlareTints)
 	{
 		for(uint32 i = 0; i < 8; ++i)
@@ -933,6 +917,30 @@ void FSceneView::OverridePostProcessSettings(const FPostProcessSettings& Src, fl
 	{
 		Dest.AntiAliasingMethod = Src.AntiAliasingMethod;
 	}
+	
+	// will be deprecated soon, use the new asset LightPropagationVolumeBlendable instead
+	{
+		FLightPropagationVolumeSettings& Dest = FinalPostProcessSettings.BlendableManager.GetSingleFinalData<FLightPropagationVolumeSettings>();
+		const float Weight = 1.0f;
+
+		LERP_PP(LPVIntensity);
+		LERP_PP(LPVSecondaryOcclusionIntensity);
+		LERP_PP(LPVSecondaryBounceIntensity);
+		LERP_PP(LPVVplInjectionBias);
+		LERP_PP(LPVGeometryVolumeBias);
+		LERP_PP(LPVEmissiveInjectionIntensity);
+		LERP_PP(LPVDirectionalOcclusionIntensity);
+		LERP_PP(LPVDirectionalOcclusionRadius);
+		LERP_PP(LPVDiffuseOcclusionExponent);
+		LERP_PP(LPVSpecularOcclusionExponent);
+		LERP_PP(LPVDiffuseOcclusionIntensity);
+		LERP_PP(LPVSpecularOcclusionIntensity);
+
+		if (Src.bOverride_LPVSize)
+		{
+			Dest.LPVSize = Src.LPVSize;
+		}
+	}
 
 	// Blendable objects
 	{
@@ -1082,14 +1090,24 @@ void FSceneView::StartFinalPostprocessSettings(FVector InViewLocation)
 
 void FSceneView::EndFinalPostprocessSettings(const FSceneViewInitOptions& ViewInitOptions)
 {
-	if(FinalPostProcessSettings.LPVDirectionalOcclusionIntensity < 0.001f)
+	// will be deprecated soon, use the new asset LightPropagationVolumeBlendable instead
 	{
-		FinalPostProcessSettings.LPVDirectionalOcclusionIntensity = 0.0f;
-	}
+		FLightPropagationVolumeSettings& Dest = FinalPostProcessSettings.BlendableManager.GetSingleFinalData<FLightPropagationVolumeSettings>();
 
-	if(FinalPostProcessSettings.LPVIntensity < 0.001f)
-	{
-		FinalPostProcessSettings.LPVIntensity = 0.0f;
+		if(Dest.LPVDirectionalOcclusionIntensity < 0.001f)
+		{
+			Dest.LPVDirectionalOcclusionIntensity = 0.0f;
+		}
+
+		if (Dest.LPVIntensity < 0.001f)
+		{
+			Dest.LPVIntensity = 0.0f;
+		}
+
+		if(!Family->EngineShowFlags.GlobalIllumination)
+		{
+			Dest.LPVIntensity = 0.0f;
+		}
 	}
 
 	{
@@ -1118,11 +1136,6 @@ void FSceneView::EndFinalPostprocessSettings(const FSceneViewInitOptions& ViewIn
 		FinalPostProcessSettings.BloomIntensity = 0.0f;
 	}
 
-	if(!Family->EngineShowFlags.GlobalIllumination)
-	{
-		FinalPostProcessSettings.LPVIntensity = 0.0f;
-	}
-
 	{
 		static const auto CVar = IConsoleManager::Get().FindTConsoleVariableDataInt(TEXT("r.DepthOfFieldQuality"));
 
diff --git a/Engine/Source/Runtime/Engine/Public/BlendableManager.h b/Engine/Source/Runtime/Engine/Public/BlendableManager.h
--- a/Engine/Source/Runtime/Engine/Public/BlendableManager.h
+++ b/Engine/Source/Runtime/Engine/Public/BlendableManager.h
@@ -43,6 +43,45 @@ public:
 		PushBlendableDataPtr(InWeight, BlendableType, (const uint8*)&InData, sizeof(T));
 	}
 
+	// used to blend multiple blendables of the given type with lerp into one final one
+	template <class T>
+	T& GetSingleFinalData()
+	{
+		FBlendableEntry* Iterator = 0;
+		T* FinalDataPtr = IterateBlendables<T>(Iterator);
+
+		if(!FinalDataPtr)
+		{
+			T Base;
+			Base.SetBaseValues();
+
+			PushBlendableData<T>(1.0f, Base);
+			// can be optimized
+			FinalDataPtr = IterateBlendables<T>(Iterator);
+		}
+		
+		return *FinalDataPtr;
+	}
+
+	// use to pickup the on final blendable that was generated by blending 0-n blendables of the given type
+	template <class T>
+	const T& GetSingleFinalDataConst() const
+	{
+		FBlendableEntry* Iterator = 0;
+		const T* FinalDataPtr = IterateBlendables<T>(Iterator);
+
+		if(!FinalDataPtr)
+		{
+			// could be improved
+			static T Base;
+			Base.SetBaseValues();
+
+			return Base;
+		}
+		
+		return *FinalDataPtr;
+	}
+
 	// find next that has the given type, O(n), n is number of elements after the given one or all if started with 0
 	// @param InIterator 0 to start iteration
 	// @return 0 if no further one of that type was found
diff --git a/Engine/Source/Runtime/Renderer/Private/CompositionLighting/CompositionLighting.cpp b/Engine/Source/Runtime/Renderer/Private/CompositionLighting/CompositionLighting.cpp
--- a/Engine/Source/Runtime/Renderer/Private/CompositionLighting/CompositionLighting.cpp
+++ b/Engine/Source/Runtime/Renderer/Private/CompositionLighting/CompositionLighting.cpp
@@ -21,6 +21,7 @@
 #include "PostProcessTemporalAA.h"
 #include "PostProcessSubsurface.h"
 #include "SceneUtils.h"
+#include "LightPropagationVolumeBlendable.h"
 
 /** The global center for all deferred lighting activities. */
 FCompositionLighting GCompositionLighting;
@@ -56,7 +57,9 @@ static bool IsLpvIndirectPassRequired(FPostprocessContext& Context)
 	{
 		FLightPropagationVolume* LightPropagationVolume = ViewState->GetLightPropagationVolume();
 
-		if(LightPropagationVolume && Context.View.FinalPostProcessSettings.LPVIntensity > 0.0f)
+		const FLightPropagationVolumeSettings& LPVSettings = Context.View.FinalPostProcessSettings.BlendableManager.GetSingleFinalDataConst<FLightPropagationVolumeSettings>();
+
+		if(LightPropagationVolume && LPVSettings.LPVIntensity > 0.0f)
 		{
 			return true; 
 		}
diff --git a/Engine/Source/Runtime/Renderer/Private/CompositionLighting/PostProcessLpvIndirect.cpp b/Engine/Source/Runtime/Renderer/Private/CompositionLighting/PostProcessLpvIndirect.cpp
--- a/Engine/Source/Runtime/Renderer/Private/CompositionLighting/PostProcessLpvIndirect.cpp
+++ b/Engine/Source/Runtime/Renderer/Private/CompositionLighting/PostProcessLpvIndirect.cpp
@@ -18,6 +18,7 @@
 #include "LightPropagationVolume.h"
 #include "UniformBuffer.h"
 #include "SceneUtils.h"
+#include "LightPropagationVolumeBlendable.h"
 
 IMPLEMENT_UNIFORM_BUFFER_STRUCT(FLpvReadUniformBufferParameters,TEXT("LpvRead"));
 typedef TUniformBufferRef<FLpvReadUniformBufferParameters> FLpvReadUniformBufferRef;
@@ -214,11 +215,12 @@ void FRCPassPostProcessLpvIndirect::Process(FRenderingCompositePassContext& Cont
 		check(PassOutputs[0].RenderTargetDesc.Extent.Y);
 	}
 
-	const FPostProcessSettings& PostprocessSettings = Context.View.FinalPostProcessSettings;
+	const FFinalPostProcessSettings& PostprocessSettings = Context.View.FinalPostProcessSettings;
 	const FSceneView& View = Context.View;
+	const FLightPropagationVolumeSettings& LPVSettings = PostprocessSettings.BlendableManager.GetSingleFinalDataConst<FLightPropagationVolumeSettings>();
 
 	FSceneViewState* ViewState = (FSceneViewState*)View.State;
-	if ( ViewState == nullptr || ViewState->GetLightPropagationVolume() == nullptr || PostprocessSettings.LPVIntensity == 0.0f )
+	if ( ViewState == nullptr || ViewState->GetLightPropagationVolume() == nullptr || LPVSettings.LPVIntensity == 0.0f )
 	{
 		return;
 	}
@@ -243,7 +245,7 @@ void FRCPassPostProcessLpvIndirect::Process(FRenderingCompositePassContext& Cont
 	// Make sure the LPV Update has completed
 	Lpv->InsertGPUWaitForAsyncUpdate(Context.RHICmdList);
 
-	if ( Context.View.FinalPostProcessSettings.LPVDirectionalOcclusionIntensity > 0.0001f )
+	if ( LPVSettings.LPVDirectionalOcclusionIntensity > 0.0001f )
 	{
 		DoDirectionalOcclusionPass(Context);
 	}
@@ -326,7 +328,7 @@ void FRCPassPostProcessLpvIndirect::Process(FRenderingCompositePassContext& Cont
 		}
 	}
 
-	if ( Context.View.FinalPostProcessSettings.LPVDirectionalOcclusionIntensity > 0.0001f )
+	if ( LPVSettings.LPVDirectionalOcclusionIntensity > 0.0001f )
 	{
 		GRenderTargetPool.VisualizeTexture.SetCheckPoint(Context.RHICmdList, GSceneRenderTargets.DirectionalOcclusion);
 	}
@@ -338,8 +340,10 @@ void FRCPassPostProcessLpvIndirect::DoDirectionalOcclusionPass(FRenderingComposi
 	const FSceneRenderTargetItem& DestDirectionalOcclusionRenderTarget = GSceneRenderTargets.DirectionalOcclusion->GetRenderTargetItem();
 	FViewInfo& View = Context.View;
 	FSceneViewState* ViewState = (FSceneViewState*)View.State;
-	const FPostProcessSettings& PostprocessSettings = Context.View.FinalPostProcessSettings;
-	if ( ViewState == nullptr || ViewState->GetLightPropagationVolume() == nullptr || PostprocessSettings.LPVIntensity == 0.0f )
+	const FFinalPostProcessSettings& PostprocessSettings = Context.View.FinalPostProcessSettings;
+	const FLightPropagationVolumeSettings& LPVSettings = PostprocessSettings.BlendableManager.GetSingleFinalDataConst<FLightPropagationVolumeSettings>();
+
+	if ( ViewState == nullptr || ViewState->GetLightPropagationVolume() == nullptr || LPVSettings.LPVIntensity == 0.0f )
 	{
 		return;
 	}
diff --git a/Engine/Source/Runtime/Renderer/Private/LightPropagationVolume.cpp b/Engine/Source/Runtime/Renderer/Private/LightPropagationVolume.cpp
--- a/Engine/Source/Runtime/Renderer/Private/LightPropagationVolume.cpp
+++ b/Engine/Source/Runtime/Renderer/Private/LightPropagationVolume.cpp
@@ -15,6 +15,7 @@
 #include "LightPropagationVolume.h"
 #include "UniformBuffer.h"
 #include "SceneUtils.h"
+#include "LightPropagationVolumeBlendable.h"
 
 static TAutoConsoleVariable<int32> CVarLightPropagationVolume(
 	TEXT("r.LightPropagationVolume"),
@@ -894,15 +895,17 @@ void FLightPropagationVolume::InitSettings(FRHICommandList& RHICmdList, const FS
 		bInitialized = true;
 	}  
 
+	const FLightPropagationVolumeSettings& LPVSettings = View.FinalPostProcessSettings.BlendableManager.GetSingleFinalDataConst<FLightPropagationVolumeSettings>();
+
 	// Copy the LPV postprocess settings
-	Strength	 = View.FinalPostProcessSettings.LPVIntensity;
+	Strength	 = LPVSettings.LPVIntensity;
 	bEnabled     = Strength > 0.0f;
-	CubeSize	 = View.FinalPostProcessSettings.LPVSize;
+	CubeSize	 = LPVSettings.LPVSize;
 
-	SecondaryOcclusionStrength = View.FinalPostProcessSettings.LPVSecondaryOcclusionIntensity;
-	SecondaryBounceStrength = View.FinalPostProcessSettings.LPVSecondaryBounceIntensity;
+	SecondaryOcclusionStrength =LPVSettings.LPVSecondaryOcclusionIntensity;
+	SecondaryBounceStrength =LPVSettings.LPVSecondaryBounceIntensity;
 
-	bGeometryVolumeNeeded = View.FinalPostProcessSettings.LPVSecondaryOcclusionIntensity > 0.001f || View.FinalPostProcessSettings.LPVDirectionalOcclusionIntensity > 0.001;
+	bGeometryVolumeNeeded =LPVSettings.LPVSecondaryOcclusionIntensity > 0.001f ||LPVSettings.LPVDirectionalOcclusionIntensity > 0.001;
 	GeometryVolumeGenerated = false;
 
 	if ( !bEnabled )
@@ -952,20 +955,20 @@ void FLightPropagationVolume::InitSettings(FRHICommandList& RHICmdList, const FS
 		LpvWriteUniformBufferParams->SecondaryBounceStrength = SecondaryBounceStrength;
 
 		// Convert the bias values from LPV grid space to world space
-		LpvWriteUniformBufferParams->GeometryVolumeInjectionBias	= View.FinalPostProcessSettings.LPVGeometryVolumeBias * LpvScale;
-		LpvWriteUniformBufferParams->VplInjectionBias				= View.FinalPostProcessSettings.LPVVplInjectionBias * LpvScale;
+		LpvWriteUniformBufferParams->GeometryVolumeInjectionBias	= LPVSettings.LPVGeometryVolumeBias * LpvScale;
+		LpvWriteUniformBufferParams->VplInjectionBias				= LPVSettings.LPVVplInjectionBias * LpvScale;
 		LpvWriteUniformBufferParams->PropagationIndex				= 0;
-		LpvWriteUniformBufferParams->EmissiveInjectionMultiplier	= View.FinalPostProcessSettings.LPVEmissiveInjectionIntensity
+		LpvWriteUniformBufferParams->EmissiveInjectionMultiplier	= LPVSettings.LPVEmissiveInjectionIntensity
 			* LpvWriteUniformBufferParams->RsmAreaIntensityMultiplier * CVarLPVEmissiveIntensityMultiplier.GetValueOnRenderThread() * 0.25f;
-		LpvWriteUniformBufferParams->DirectionalOcclusionIntensity	= View.FinalPostProcessSettings.LPVDirectionalOcclusionIntensity;
-		LpvWriteUniformBufferParams->DirectionalOcclusionRadius		= View.FinalPostProcessSettings.LPVDirectionalOcclusionRadius;
+		LpvWriteUniformBufferParams->DirectionalOcclusionIntensity	= LPVSettings.LPVDirectionalOcclusionIntensity;
+		LpvWriteUniformBufferParams->DirectionalOcclusionRadius		= LPVSettings.LPVDirectionalOcclusionRadius;
 		LpvWriteUniformBufferParams->RsmPixelToTexcoordMultiplier	= 1.0f / float(GSceneRenderTargets.GetReflectiveShadowMapResolution() - 1);
 
-		LpvReadUniformBufferParams.DirectionalOcclusionIntensity	= View.FinalPostProcessSettings.LPVDirectionalOcclusionIntensity;
-		LpvReadUniformBufferParams.DiffuseOcclusionExponent			= View.FinalPostProcessSettings.LPVDiffuseOcclusionExponent;
-		LpvReadUniformBufferParams.SpecularOcclusionExponent		= View.FinalPostProcessSettings.LPVSpecularOcclusionExponent;
-		LpvReadUniformBufferParams.DiffuseOcclusionIntensity		= View.FinalPostProcessSettings.LPVDiffuseOcclusionIntensity;
-		LpvReadUniformBufferParams.SpecularOcclusionIntensity		= View.FinalPostProcessSettings.LPVSpecularOcclusionIntensity;
+		LpvReadUniformBufferParams.DirectionalOcclusionIntensity	= LPVSettings.LPVDirectionalOcclusionIntensity;
+		LpvReadUniformBufferParams.DiffuseOcclusionExponent			= LPVSettings.LPVDiffuseOcclusionExponent;
+		LpvReadUniformBufferParams.SpecularOcclusionExponent		= LPVSettings.LPVSpecularOcclusionExponent;
+		LpvReadUniformBufferParams.DiffuseOcclusionIntensity		= LPVSettings.LPVDiffuseOcclusionIntensity;
+		LpvReadUniformBufferParams.SpecularOcclusionIntensity		= LPVSettings.LPVSpecularOcclusionIntensity;
 
 		LpvReadUniformBufferParams.mLpvGridOffset		= mGridOffset;
 		LpvReadUniformBufferParams.LpvScale				= LpvScale;
diff --git a/Engine/Source/Runtime/Renderer/Private/PostProcess/SceneRenderTargets.cpp b/Engine/Source/Runtime/Renderer/Private/PostProcess/SceneRenderTargets.cpp
--- a/Engine/Source/Runtime/Renderer/Private/PostProcess/SceneRenderTargets.cpp
+++ b/Engine/Source/Runtime/Renderer/Private/PostProcess/SceneRenderTargets.cpp
@@ -11,6 +11,7 @@
 #include "LightPropagationVolume.h"
 #include "SceneUtils.h"
 #include "HdrCustomResolveShaders.h"
+#include "Public/LightPropagationVolumeBlendable.h"
 
 IMPLEMENT_UNIFORM_BUFFER_STRUCT(FGBufferResourceStruct,TEXT("GBuffers"));
 
@@ -20,7 +21,7 @@ static TAutoConsoleVariable<int32> CVarBasePassOutputsVelocityDebug(
 	TEXT("Debug settings for Base Pass outputting velocity.\n")
 	TEXT("0 - Regular rendering\n")
 	TEXT("1 - Skip setting GBufferVelocity RT\n")
-	TEXT("2 - Set Color Mask 0 for GBufferVelocity RT\n"),
+	TEXT("2 - Set Color Mask 0 for GBufferVelocity RT"),
 	ECVF_RenderThreadSafe);
 
 static TAutoConsoleVariable<int32> CVarRSMResolution(
@@ -2015,9 +2016,12 @@ void FSceneTextureShaderParameters::Set(
 	{
 		bool bDirectionalOcclusion = false;
 		FSceneViewState* ViewState = (FSceneViewState*)View.State;
-		if ( ViewState != nullptr && ViewState->GetLightPropagationVolume() != nullptr && View.FinalPostProcessSettings.LPVIntensity > 0.0f )
+
+		const FLightPropagationVolumeSettings& LPVSettings = View.FinalPostProcessSettings.BlendableManager.GetSingleFinalDataConst<FLightPropagationVolumeSettings>();
+
+		if ( ViewState != nullptr && ViewState->GetLightPropagationVolume() != nullptr && LPVSettings.LPVIntensity > 0.0f )
 		{
-			if ( View.FinalPostProcessSettings.LPVDirectionalOcclusionIntensity > 0.0001f )
+			if ( LPVSettings.LPVDirectionalOcclusionIntensity > 0.0001f )
 			{
 				bDirectionalOcclusion = true;
 			}
diff --git a/Engine/Source/Runtime/Renderer/Private/ReflectionEnvironment.cpp b/Engine/Source/Runtime/Renderer/Private/ReflectionEnvironment.cpp
--- a/Engine/Source/Runtime/Renderer/Private/ReflectionEnvironment.cpp
+++ b/Engine/Source/Runtime/Renderer/Private/ReflectionEnvironment.cpp
@@ -972,7 +972,10 @@ void FDeferredShadingSceneRenderer::RenderStandardDeferredImageBasedReflections(
 
 		bool bLPV = false;
 		FSceneViewState* ViewState = (FSceneViewState*)View.State;
-		if ( ViewState != nullptr && ViewState->GetLightPropagationVolume() != nullptr && View.FinalPostProcessSettings.LPVIntensity > 0.0f )
+		
+		const FLightPropagationVolumeSettings& LPVSettings = View.FinalPostProcessSettings.BlendableManager.GetSingleFinalDataConst<FLightPropagationVolumeSettings>();
+
+		if ( ViewState != nullptr && ViewState->GetLightPropagationVolume() != nullptr && LPVSettings.LPVIntensity > 0.0f )
 		{
 			bLPV = true;
 		} 
diff --git a/Engine/Source/Runtime/Renderer/Private/ShadowSetup.cpp b/Engine/Source/Runtime/Renderer/Private/ShadowSetup.cpp
--- a/Engine/Source/Runtime/Renderer/Private/ShadowSetup.cpp
+++ b/Engine/Source/Runtime/Renderer/Private/ShadowSetup.cpp
@@ -7,6 +7,7 @@
 #include "RendererPrivate.h"
 #include "ScenePrivate.h"
 #include "LightPropagationVolume.h"
+#include "LightPropagationVolumeBlendable.h"
 
 static float GMinScreenRadiusForShadowCaster = 0.03f;
 static FAutoConsoleVariableRef CVarMinScreenRadiusForShadowCaster(
@@ -2246,8 +2247,10 @@ void FSceneRenderer::AddViewDependentWholeSceneShadowsForView(
 			if (ViewState)
 			{
 				FLightPropagationVolume* LightPropagationVolume = ViewState->GetLightPropagationVolume();
+				
+				FLightPropagationVolumeSettings& LPVSettings = View.FinalPostProcessSettings.BlendableManager.GetSingleFinalData<FLightPropagationVolumeSettings>();
 
-				if (LightPropagationVolume && LightPropagationVolume->bInitialized && View.FinalPostProcessSettings.LPVIntensity > 0)
+				if (LightPropagationVolume && LightPropagationVolume->bInitialized && LPVSettings.LPVIntensity > 0)
 				{
 					// Generate the RSM shadow info
 					FWholeSceneProjectedShadowInitializer ProjectedShadowInitializer;
diff --git a/Engine/Source/Runtime/Renderer/Renderer.Build.cs b/Engine/Source/Runtime/Renderer/Renderer.Build.cs
--- a/Engine/Source/Runtime/Renderer/Renderer.Build.cs
+++ b/Engine/Source/Runtime/Renderer/Renderer.Build.cs
@@ -32,9 +32,9 @@ public class Renderer : ModuleRules
 				"ShaderCore",
 				"UtilityShaders",
 			}
-			);
+            );
 
-        PrivateIncludePathModuleNames.AddRange(new string[] { "HeadMountedDisplay" });
+        PrivateIncludePathModuleNames.AddRange(new string[] { "HeadMountedDisplay", "LightPropagationVolumeRuntime" });
         DynamicallyLoadedModuleNames.AddRange(new string[] { "HeadMountedDisplay" });
 	}
 }
