UE4 - fix CIS

diff --git a/Engine/Plugins/Blendables/LightPropagationVolume/Source/LightPropagationVolumeRuntime/Private/LightPropagationVolumeRuntimeModule.cpp b/Engine/Plugins/Blendables/LightPropagationVolume/Source/LightPropagationVolumeRuntime/Private/LightPropagationVolumeRuntimeModule.cpp
--- a/Engine/Plugins/Blendables/LightPropagationVolume/Source/LightPropagationVolumeRuntime/Private/LightPropagationVolumeRuntimeModule.cpp
+++ b/Engine/Plugins/Blendables/LightPropagationVolume/Source/LightPropagationVolumeRuntime/Private/LightPropagationVolumeRuntimeModule.cpp
@@ -5,7 +5,9 @@
 =============================================================================*/
 
 #include "Core.h"
+#include "ModuleManager.h"
 #include "LightPropagationVolumeBlendable.h"
+#include "SceneView.h"
 
 IMPLEMENT_MODULE(FDefaultModuleImpl, LightPropagationVolumeRuntime);
 
diff --git a/Engine/Plugins/Blendables/LightPropagationVolume/Source/LightPropagationVolumeRuntime/Public/LightPropagationVolumeBlendable.h b/Engine/Plugins/Blendables/LightPropagationVolume/Source/LightPropagationVolumeRuntime/Public/LightPropagationVolumeBlendable.h
--- a/Engine/Plugins/Blendables/LightPropagationVolume/Source/LightPropagationVolumeRuntime/Public/LightPropagationVolumeBlendable.h
+++ b/Engine/Plugins/Blendables/LightPropagationVolume/Source/LightPropagationVolumeRuntime/Public/LightPropagationVolumeBlendable.h
@@ -2,6 +2,7 @@
 
 #pragma once
 
+#include "CoreUObject.h"
 #include "Engine/BlendableInterface.h"
 #include "Engine/DataAsset.h"
 #include "LightPropagationVolumeBlendable.generated.h"
diff --git a/Engine/Source/Runtime/Engine/Classes/Engine/BlendableInterface.h b/Engine/Source/Runtime/Engine/Classes/Engine/BlendableInterface.h
--- a/Engine/Source/Runtime/Engine/Classes/Engine/BlendableInterface.h
+++ b/Engine/Source/Runtime/Engine/Classes/Engine/BlendableInterface.h
@@ -1,9 +1,10 @@
 // Copyright 1998-2015 Epic Games, Inc. All Rights Reserved.
 
 #pragma once
-
+#include "Interface.h"
 #include "BlendableInterface.generated.h"
 
+class UMaterialInterface;
 class UMaterialInstanceDynamic;
 
 /** Where to place a material node in the post processing graph */
