UE4 - maybe fix CIS

diff --git a/Engine/Plugins/Blendables/LightPropagationVolume/Source/LightPropagationVolumeEditor/Classes/LightPropagationVolumeBlendableFactory.h b/Engine/Plugins/Blendables/LightPropagationVolume/Source/LightPropagationVolumeEditor/Classes/LightPropagationVolumeBlendableFactory.h
--- a/Engine/Plugins/Blendables/LightPropagationVolume/Source/LightPropagationVolumeEditor/Classes/LightPropagationVolumeBlendableFactory.h
+++ b/Engine/Plugins/Blendables/LightPropagationVolume/Source/LightPropagationVolumeEditor/Classes/LightPropagationVolumeBlendableFactory.h
@@ -1,7 +1,8 @@
 // Copyright 1998-2015 Epic Games, Inc. All Rights Reserved.
 
 #pragma once
-
+#include "CoreUObject.h"
+#include "Factories/Factory.h"
 #include "LightPropagationVolumeBlendableFactory.generated.h"
 
 UCLASS(hidecategories=Object)
