fixed Clang errors

diff --git a/Engine/Plugins/Blendables/LightPropagationVolume/Source/LightPropagationVolumeEditor/Private/LightPropagationVolumeEditorModule.cpp b/Engine/Plugins/Blendables/LightPropagationVolume/Source/LightPropagationVolumeEditor/Private/LightPropagationVolumeEditorModule.cpp
--- a/Engine/Plugins/Blendables/LightPropagationVolume/Source/LightPropagationVolumeEditor/Private/LightPropagationVolumeEditorModule.cpp
+++ b/Engine/Plugins/Blendables/LightPropagationVolume/Source/LightPropagationVolumeEditor/Private/LightPropagationVolumeEditorModule.cpp
@@ -4,7 +4,7 @@
 	LightPropagationVolumeModule.cpp: Module encapsulates the asset type for blending LightPropagationVolume settings
 =============================================================================*/
 
-#include "Core.h"
+#include "UnrealEd.h"
 #include "AssetTypeActions_LightPropagationVolumeBlendable.h"
 #include "LightPropagationVolumeBlendableFactory.h"
 #include "LightPropagationVolumeBlendable.h"
