Remove CanContainContent: "true" for plugins that do not have content folders, to prevent warnings on startup. (UE-17486)

diff --git a/Engine/Plugins/Blendables/LightPropagationVolume/LightPropagationVolume.uplugin b/Engine/Plugins/Blendables/LightPropagationVolume/LightPropagationVolume.uplugin
--- a/Engine/Plugins/Blendables/LightPropagationVolume/LightPropagationVolume.uplugin
+++ b/Engine/Plugins/Blendables/LightPropagationVolume/LightPropagationVolume.uplugin
@@ -21,7 +21,5 @@
 			"Name" : "LightPropagationVolumeEditor",
 			"Type" : "Editor"
 		},
-    ],
-
-	"CanContainContent" : true
+    ]
 }
\ No newline at end of file
diff --git a/Engine/Plugins/Experimental/CharacterAI/CharacterAI.uplugin b/Engine/Plugins/Experimental/CharacterAI/CharacterAI.uplugin
--- a/Engine/Plugins/Experimental/CharacterAI/CharacterAI.uplugin
+++ b/Engine/Plugins/Experimental/CharacterAI/CharacterAI.uplugin
@@ -19,7 +19,5 @@
 			"Type" : "Runtime",
 			"LoadingPhase" : "PostEngineInit"
 		},
-	],
-
-	"CanContainContent" : true
+	]
 }
\ No newline at end of file
