added plugin description

diff --git a/Engine/Plugins/Blendables/LightPropagationVolume/LightPropagationVolume.uplugin b/Engine/Plugins/Blendables/LightPropagationVolume/LightPropagationVolume.uplugin
--- a/Engine/Plugins/Blendables/LightPropagationVolume/LightPropagationVolume.uplugin
+++ b/Engine/Plugins/Blendables/LightPropagationVolume/LightPropagationVolume.uplugin
@@ -7,8 +7,8 @@
     "CreatedBy": "Epic Games, Inc.",
     "CreatedByURL": "http://epicgames.com",
     "EngineVersion": "4.2.0",
-    "Description": "todo",
-    "Category": "3D",
+    "Description": "The plugin encapsulates a small part of the LightPropagationVolume feature, the LightPropagationVolumeBlendable asset (for more info: look for Blendables in the documentation). This plugin is also an example on how to make a Blendable asset outside of the engine",
+    "Category": "Rendering",
     "EnabledByDefault": true,
 
     "Modules": [
